Compiled UnrealVS in VS2019 to make it load correctly in VS2019.

Bumped version number to trigger a reinstall.

Updated release notes.

diff --git a/Engine/Extras/UnrealVS/VS2022/AtomFeed.xml b/Engine/Extras/UnrealVS/VS2022/AtomFeed.xml
--- a/Engine/Extras/UnrealVS/VS2022/AtomFeed.xml
+++ b/Engine/Extras/UnrealVS/VS2022/AtomFeed.xml
@@ -8,7 +8,7 @@
     <title type="text">UnrealVS</title>
     <summary type="text">Unreal Engine extension for Visual Studio</summary>
     <published>2022-10-21T10:01:57-07:00</published>
-    <updated>2022-10-28T11:09:39-07:00</updated>
+    <updated>2023-01-31T13:00:00-00:00</updated>
     <author>
       <name>Epic Games, Inc.</name>
     </author>
@@ -19,7 +19,7 @@
     <content type="application/octet-stream" src="UnrealVS.vsix" />
     <Vsix xmlns:xsd="http://www.w3.org/2001/XMLSchema" xmlns:xsi="http://www.w3.org/2001/XMLSchema-instance" xmlns="http://schemas.microsoft.com/developer/vsx-syndication-schema/2010">
       <Id>ddbf523f-7eb6-4887-bd51-85a714ff87eb</Id>
-      <Version>1.80</Version>
+      <Version>1.83</Version>
       <Rating xsi:nil="true" />
       <RatingCount xsi:nil="true" />
       <DownloadCount xsi:nil="true" />
diff --git a/Engine/Source/Programs/UnrealVS/ReleaseNotes.txt b/Engine/Source/Programs/UnrealVS/ReleaseNotes.txt
--- a/Engine/Source/Programs/UnrealVS/ReleaseNotes.txt
+++ b/Engine/Source/Programs/UnrealVS/ReleaseNotes.txt
@@ -1,3 +1,9 @@
+1.83
+* Open for Edit now also Adds if the file is not already added
+* added Revert command
+* the File History command no longer hangs
+* the Timelapse command now scrolls to the current line
+
 1.80
 * updated Unreal.TestAdapater to report duration's
 * added a menu for extension to improve discover-ability of features
diff --git a/Engine/Source/Programs/UnrealVS/Unreal.TestAdapter/packages.config b/Engine/Source/Programs/UnrealVS/Unreal.TestAdapter/packages.config
new file mode 100644
--- /dev/null
+++ b/Engine/Source/Programs/UnrealVS/Unreal.TestAdapter/packages.config
@@ -0,0 +1,8 @@
+﻿<?xml version="1.0" encoding="utf-8"?>
+<packages>
+  <package id="Microsoft.TestPlatform" version="17.4.1" targetFramework="net472" />
+  <package id="Microsoft.TestPlatform.ObjectModel" version="17.4.1" targetFramework="net472" />
+  <package id="NuGet.Frameworks" version="5.11.3" targetFramework="net472" />
+  <package id="System.Collections.Immutable" version="1.5.0" targetFramework="net472" />
+  <package id="System.Reflection.Metadata" version="1.6.0" targetFramework="net472" />
+</packages>
\ No newline at end of file
