fix to UnrealVS for leaving in debugging code
https://jira.it.epicgames.com/browse/UE-179578

diff --git a/Engine/Extras/UnrealVS/VS2022/AtomFeed.xml b/Engine/Extras/UnrealVS/VS2022/AtomFeed.xml
--- a/Engine/Extras/UnrealVS/VS2022/AtomFeed.xml
+++ b/Engine/Extras/UnrealVS/VS2022/AtomFeed.xml
@@ -8,7 +8,7 @@
     <title type="text">UnrealVS</title>
     <summary type="text">Unreal Engine extension for Visual Studio</summary>
     <published>2022-10-21T10:01:57-07:00</published>
-    <updated>2023-03-07T13:10:00-00:00</updated>
+    <updated>2023-03-14T13:10:00-00:00</updated>
     <author>
       <name>Epic Games, Inc.</name>
     </author>
@@ -19,7 +19,7 @@
     <content type="application/octet-stream" src="UnrealVS.vsix" />
     <Vsix xmlns:xsd="http://www.w3.org/2001/XMLSchema" xmlns:xsi="http://www.w3.org/2001/XMLSchema-instance" xmlns="http://schemas.microsoft.com/developer/vsx-syndication-schema/2010">
       <Id>ddbf523f-7eb6-4887-bd51-85a714ff87eb</Id>
-      <Version>1.84.3</Version>
+      <Version>1.84.3.1</Version>
       <Rating xsi:nil="true" />
       <RatingCount xsi:nil="true" />
       <DownloadCount xsi:nil="true" />
diff --git a/Engine/Source/Programs/UnrealVS/ReleaseNotes.txt b/Engine/Source/Programs/UnrealVS/ReleaseNotes.txt
--- a/Engine/Source/Programs/UnrealVS/ReleaseNotes.txt
+++ b/Engine/Source/Programs/UnrealVS/ReleaseNotes.txt
@@ -1,3 +1,6 @@
+1.84.3.1
+* Bug fix for leaving in debugging code
+
 1.84.3
 * Added command to generate assembly file while compiling single file
 
diff --git a/Engine/Source/Programs/UnrealVS/Unreal.TestAdapter/Properties/AssemblyInfo.cs b/Engine/Source/Programs/UnrealVS/Unreal.TestAdapter/Properties/AssemblyInfo.cs
--- a/Engine/Source/Programs/UnrealVS/Unreal.TestAdapter/Properties/AssemblyInfo.cs
+++ b/Engine/Source/Programs/UnrealVS/Unreal.TestAdapter/Properties/AssemblyInfo.cs
@@ -33,5 +33,5 @@ using System.Runtime.InteropServices;
 // You can specify all the values or you can default the Build and Revision Numbers
 // by using the '*' as shown below:
 // [assembly: AssemblyVersion("1.0.*")]
-[assembly: AssemblyVersion("1.1.0.0")]
-[assembly: AssemblyFileVersion("1.1.0.0")]
+[assembly: AssemblyVersion("1.1.0.1")]
+[assembly: AssemblyFileVersion("1.1.0.1")]
diff --git a/Engine/Source/Programs/UnrealVS/Unreal.TestAdapter/TestExecutor.cs b/Engine/Source/Programs/UnrealVS/Unreal.TestAdapter/TestExecutor.cs
--- a/Engine/Source/Programs/UnrealVS/Unreal.TestAdapter/TestExecutor.cs
+++ b/Engine/Source/Programs/UnrealVS/Unreal.TestAdapter/TestExecutor.cs
@@ -69,7 +69,6 @@ namespace Unreal.TestAdapter
 			}
 			else
 			{
-				System.Diagnostics.Debugger.Launch();
 				string tempFile = Path.Combine(Path.GetTempPath(), Path.GetTempFileName());
 				using (Process process = new Process())
 				{
diff --git a/Engine/Source/Programs/UnrealVS/UnrealVS.2019/Properties/AssemblyInfo.cs b/Engine/Source/Programs/UnrealVS/UnrealVS.2019/Properties/AssemblyInfo.cs
--- a/Engine/Source/Programs/UnrealVS/UnrealVS.2019/Properties/AssemblyInfo.cs
+++ b/Engine/Source/Programs/UnrealVS/UnrealVS.2019/Properties/AssemblyInfo.cs
@@ -31,6 +31,6 @@ using System.Runtime.InteropServices;
 // You can specify all the values or you can default the Revision and Build Numbers 
 // by using the '*' as shown below:
 
-[assembly: AssemblyVersion("1.84.3.0")]
-[assembly: AssemblyFileVersion("1.84.3.0")]
+[assembly: AssemblyVersion("1.84.3.1")]
+[assembly: AssemblyFileVersion("1.84.3.1")]
 
