[UnrealVS] Fixed Visual Studio being unresponsive due to unexpected load order with initialization. Only updating VS2022 plugin since VS2019 is no longer supported.

diff --git a/Engine/Extras/UnrealVS/VS2022/AtomFeed.xml b/Engine/Extras/UnrealVS/VS2022/AtomFeed.xml
--- a/Engine/Extras/UnrealVS/VS2022/AtomFeed.xml
+++ b/Engine/Extras/UnrealVS/VS2022/AtomFeed.xml
@@ -8,7 +8,7 @@
     <title type="text">UnrealVS</title>
     <summary type="text">Unreal Engine extension for Visual Studio</summary>
     <published>2022-10-21T10:01:57-07:00</published>
-    <updated>2024-03-15T18:30:00-00:00</updated>
+    <updated>2024-04-05T18:30:00-00:00</updated>
     <author>
       <name>Epic Games, Inc.</name>
     </author>
@@ -19,7 +19,7 @@
     <content type="application/octet-stream" src="UnrealVS.vsix" />
     <Vsix xmlns:xsd="http://www.w3.org/2001/XMLSchema" xmlns:xsi="http://www.w3.org/2001/XMLSchema-instance" xmlns="http://schemas.microsoft.com/developer/vsx-syndication-schema/2010">
       <Id>ddbf523f-7eb6-4887-bd51-85a714ff87eb</Id>
-      <Version>1.84.3.9</Version>
+      <Version>1.84.4.0</Version>
       <Rating xsi:nil="true" />
       <RatingCount xsi:nil="true" />
       <DownloadCount xsi:nil="true" />
diff --git a/Engine/Source/Programs/UnrealVS/ReleaseNotes.txt b/Engine/Source/Programs/UnrealVS/ReleaseNotes.txt
--- a/Engine/Source/Programs/UnrealVS/ReleaseNotes.txt
+++ b/Engine/Source/Programs/UnrealVS/ReleaseNotes.txt
@@ -1,3 +1,6 @@
+1.84.4.0
+* Fixed a race condition at initialization time that would make Visual Studio unresponsive.
+
 1.84.3.9
 * Fixed issues with compile single file with profile. Did not work if selected project was a c# project. Should now also work with non-ue projects
 
diff --git a/Engine/Source/Programs/UnrealVS/UnrealVS.2019/Properties/AssemblyInfo.cs b/Engine/Source/Programs/UnrealVS/UnrealVS.2019/Properties/AssemblyInfo.cs
--- a/Engine/Source/Programs/UnrealVS/UnrealVS.2019/Properties/AssemblyInfo.cs
+++ b/Engine/Source/Programs/UnrealVS/UnrealVS.2019/Properties/AssemblyInfo.cs
@@ -31,6 +31,6 @@ using System.Runtime.InteropServices;
 // You can specify all the values or you can default the Revision and Build Numbers 
 // by using the '*' as shown below:
 
-[assembly: AssemblyVersion("1.84.3.9")]
-[assembly: AssemblyFileVersion("1.84.3.9")]
+[assembly: AssemblyVersion("1.84.4.0")]
+[assembly: AssemblyFileVersion("1.84.4.0")]
 
diff --git a/Engine/Source/Programs/UnrealVS/UnrealVS.2022/UnrealVS.2022.csproj b/Engine/Source/Programs/UnrealVS/UnrealVS.2022/UnrealVS.2022.csproj
--- a/Engine/Source/Programs/UnrealVS/UnrealVS.2022/UnrealVS.2022.csproj
+++ b/Engine/Source/Programs/UnrealVS/UnrealVS.2022/UnrealVS.2022.csproj
@@ -217,6 +217,7 @@
     <CopyVsixExtensionFiles>False</CopyVsixExtensionFiles>
     <NoWarn>MSB3277</NoWarn>
     <TreatWarningsAsErrors>true</TreatWarningsAsErrors>
+    <WarningLevel>5</WarningLevel>
   </PropertyGroup>
   <Target Name="IncludeNuGetPackageReferences" AfterTargets="GetVsixSourceItems">
     <ItemGroup>
diff --git a/Engine/Source/Programs/UnrealVS/UnrealVS.Shared/QuickBuild.cs b/Engine/Source/Programs/UnrealVS/UnrealVS.Shared/QuickBuild.cs
--- a/Engine/Source/Programs/UnrealVS/UnrealVS.Shared/QuickBuild.cs
+++ b/Engine/Source/Programs/UnrealVS/UnrealVS.Shared/QuickBuild.cs
@@ -109,7 +109,7 @@ namespace UnrealVS
 
 			// Get the project clicked in the solution explorer by accessing the current selection and converting to a Project if possible.
 			UnrealVSPackage.Instance.SelectionManager.GetCurrentSelection(out IntPtr HierarchyPtr, out uint ProjectItemId, out IVsMultiItemSelect MultiItemSelect, out IntPtr SelectionContainerPtr);
-			if (HierarchyPtr == null) return;
+			if (HierarchyPtr == IntPtr.Zero) return;
 
 			if (!(Marshal.GetTypedObjectForIUnknown(HierarchyPtr, typeof(IVsHierarchy)) is IVsHierarchy SelectedHierarchy)) return;
 
diff --git a/Engine/Source/Programs/UnrealVS/UnrealVS.Shared/UnrealVSPackage.cs b/Engine/Source/Programs/UnrealVS/UnrealVS.Shared/UnrealVSPackage.cs
--- a/Engine/Source/Programs/UnrealVS/UnrealVS.Shared/UnrealVSPackage.cs
+++ b/Engine/Source/Programs/UnrealVS/UnrealVS.Shared/UnrealVSPackage.cs
@@ -222,6 +222,11 @@ namespace UnrealVS
 			Assumes.Present(DTE);
 			Logging.WriteLine("DTE version " + DTE.Version);
 
+			// We previously deferred getting the DTE2 object but there are some cases where it's not running yet, so try here
+			_DTE2 = await GetServiceAsync(typeof(SDTE)) as DTE2;
+			Assumes.Present(_DTE2);
+			Assumes.True(_DTE2.DTE == DTE);
+
 			//TextManager = await GetServiceAsync(typeof(VsTextManagerClass)) as IVsTextManager3;
 
 			var componentModel = (IComponentModel)GetGlobalService(typeof(SComponentModel));
