[UnrealVs]
* Speculative fix for exception happening when opening crashdump in visual studio
* Minor optimization to reduce some flickering

diff --git a/Engine/Extras/UnrealVS/VS2022/AtomFeed.xml b/Engine/Extras/UnrealVS/VS2022/AtomFeed.xml
--- a/Engine/Extras/UnrealVS/VS2022/AtomFeed.xml
+++ b/Engine/Extras/UnrealVS/VS2022/AtomFeed.xml
@@ -8,7 +8,7 @@
     <title type="text">UnrealVS</title>
     <summary type="text">Unreal Engine extension for Visual Studio</summary>
     <published>2022-10-21T10:01:57-07:00</published>
-    <updated>2024-11-27T09:37:00-00:00</updated>
+    <updated>2025-01-06T20:48:00-00:00</updated>
     <author>
       <name>Epic Games, Inc.</name>
     </author>
@@ -19,7 +19,7 @@
     <content type="application/octet-stream" src="UnrealVS.vsix" />
     <Vsix xmlns:xsd="http://www.w3.org/2001/XMLSchema" xmlns:xsi="http://www.w3.org/2001/XMLSchema-instance" xmlns="http://schemas.microsoft.com/developer/vsx-syndication-schema/2010">
       <Id>ddbf523f-7eb6-4887-bd51-85a714ff87eb</Id>
-      <Version>1.84.8.0</Version>
+      <Version>1.84.8.1</Version>
       <Rating xsi:nil="true" />
       <RatingCount xsi:nil="true" />
       <DownloadCount xsi:nil="true" />
diff --git a/Engine/Source/Programs/UnrealVS/ReleaseNotes.txt b/Engine/Source/Programs/UnrealVS/ReleaseNotes.txt
--- a/Engine/Source/Programs/UnrealVS/ReleaseNotes.txt
+++ b/Engine/Source/Programs/UnrealVS/ReleaseNotes.txt
@@ -1,3 +1,6 @@
+1.84.8.1
+* Bug fix for UbaVisualizer throwing exception when opening crash dumps
+
 1.84.8.0
 * Fix dpi scaling support for UBA
 
diff --git a/Engine/Source/Programs/UnrealVS/UnrealVS.2022/Properties/AssemblyInfo.cs b/Engine/Source/Programs/UnrealVS/UnrealVS.2022/Properties/AssemblyInfo.cs
--- a/Engine/Source/Programs/UnrealVS/UnrealVS.2022/Properties/AssemblyInfo.cs
+++ b/Engine/Source/Programs/UnrealVS/UnrealVS.2022/Properties/AssemblyInfo.cs
@@ -31,6 +31,6 @@ using System.Runtime.InteropServices;
 // You can specify all the values or you can default the Revision and Build Numbers 
 // by using the '*' as shown below:
 
-[assembly: AssemblyVersion("1.84.8.0")]
-[assembly: AssemblyFileVersion("1.84.8.0")]
+[assembly: AssemblyVersion("1.84.8.1")]
+[assembly: AssemblyFileVersion("1.84.8.1")]
 
diff --git a/Engine/Source/Programs/UnrealVS/UnrealVS.Shared/UbaVisualizer.cs b/Engine/Source/Programs/UnrealVS/UnrealVS.Shared/UbaVisualizer.cs
--- a/Engine/Source/Programs/UnrealVS/UnrealVS.Shared/UbaVisualizer.cs
+++ b/Engine/Source/Programs/UnrealVS/UnrealVS.Shared/UbaVisualizer.cs
@@ -110,6 +110,7 @@ namespace UnrealVS
 		{
 			Window = window;
 			VisualizerExe = visualizerExe;
+			ClipToBounds = false;
 		}
 
 		protected override HandleRef BuildWindowCore(HandleRef hwndParent)
@@ -183,10 +184,14 @@ namespace UnrealVS
 					//Focus();
 					//NativeMethods.SetActiveWindow(hwnd);
 				}
+				else if (msg == 0x0014) // WM_ERASEBKGND
+				{
+					return new IntPtr(1);
+				}
 			}
 			else
 			{
-				if (msg == 0x0444) // Special UbaVisualizer message
+				if (msg == 0x0444) // Special UbaVisualizer message. Called when visualizer process has set this hwnd as parent
 				{
 					NativeMethods.DestroyWindow(HwndStatic);
 					HwndStatic = IntPtr.Zero;
diff --git a/Engine/Source/Programs/UnrealVS/UnrealVS.Shared/UbaVisualizerWindowControl.xaml.cs b/Engine/Source/Programs/UnrealVS/UnrealVS.Shared/UbaVisualizerWindowControl.xaml.cs
--- a/Engine/Source/Programs/UnrealVS/UnrealVS.Shared/UbaVisualizerWindowControl.xaml.cs
+++ b/Engine/Source/Programs/UnrealVS/UnrealVS.Shared/UbaVisualizerWindowControl.xaml.cs
@@ -84,8 +84,17 @@ namespace UnrealVS
 			}
 
 			string fullPath = path.Substring(0, engineIndex + engineStr.Length) + @"Binaries\Win64\UnrealBuildAccelerator\x64\UbaVisualizer.exe";
-			FileInfo fi = new FileInfo(fullPath);
-			if (!fi.Exists)
+			bool exists = false;
+			try
+			{
+				FileInfo fi = new FileInfo(fullPath);
+				exists = fi.Exists;
+			}
+			catch
+			{
+			}
+
+			if (!exists)
 			{
 				SetChild(null, DefaultText);
 				return;
