Version string incorrect as compared to other metadata

diff --git a/Engine/Extras/UnrealVS/VS2022/AtomFeed.xml b/Engine/Extras/UnrealVS/VS2022/AtomFeed.xml
--- a/Engine/Extras/UnrealVS/VS2022/AtomFeed.xml
+++ b/Engine/Extras/UnrealVS/VS2022/AtomFeed.xml
@@ -8,7 +8,7 @@
     <title type="text">UnrealVS</title>
     <summary type="text">Unreal Engine extension for Visual Studio</summary>
     <published>2022-10-21T10:01:57-07:00</published>
-    <updated>2025-01-06T20:48:00-00:00</updated>
+    <updated>2025-02-04T10:00:00-00:00</updated>
     <author>
       <name>Epic Games, Inc.</name>
     </author>
@@ -19,7 +19,7 @@
     <content type="application/octet-stream" src="UnrealVS.vsix" />
     <Vsix xmlns:xsd="http://www.w3.org/2001/XMLSchema" xmlns:xsi="http://www.w3.org/2001/XMLSchema-instance" xmlns="http://schemas.microsoft.com/developer/vsx-syndication-schema/2010">
       <Id>ddbf523f-7eb6-4887-bd51-85a714ff87eb</Id>
-      <Version>1.84.8.1</Version>
+      <Version>1.84.8.2</Version>
       <Rating xsi:nil="true" />
       <RatingCount xsi:nil="true" />
       <DownloadCount xsi:nil="true" />
diff --git a/Engine/Source/Programs/UnrealVS/README b/Engine/Source/Programs/UnrealVS/README
--- a/Engine/Source/Programs/UnrealVS/README
+++ b/Engine/Source/Programs/UnrealVS/README
@@ -11,4 +11,11 @@ Point the URL to the AtomFeed of where UnrealVS.vsix is synced
 ex.
 
 Name: UnrealVS
-URL: file:///D:/dev/ue5main/Engine/Extras/UnrealVS/VS2022/AtomFeed.xml
\ No newline at end of file
+URL: file:///D:/dev/ue5main/Engine/Extras/UnrealVS/VS2022/AtomFeed.xml
+
+# Debugging
+In order to debug UnrealVS.vsix:
+- Set the Startup project to UnrealVS.2022
+- Update the Debug properties for the UnrealVS.2022 csproj
+    - Start external program set to: C:\Program Files\Microsoft Visual Studio\2022\Professional\Common7\IDE\devenv.exe
+    - Command line arguments set to: /rootsuffix Exp
\ No newline at end of file
diff --git a/Engine/Source/Programs/UnrealVS/ReleaseNotes.txt b/Engine/Source/Programs/UnrealVS/ReleaseNotes.txt
--- a/Engine/Source/Programs/UnrealVS/ReleaseNotes.txt
+++ b/Engine/Source/Programs/UnrealVS/ReleaseNotes.txt
@@ -1,3 +1,6 @@
+1.84.8.2
+* Bug fix where the package info string is incorrect
+
 1.84.8.1
 * Bug fix for UbaVisualizer throwing exception when opening crash dumps
 
diff --git a/Engine/Source/Programs/UnrealVS/UnrealVS.2022/Properties/AssemblyInfo.cs b/Engine/Source/Programs/UnrealVS/UnrealVS.2022/Properties/AssemblyInfo.cs
--- a/Engine/Source/Programs/UnrealVS/UnrealVS.2022/Properties/AssemblyInfo.cs
+++ b/Engine/Source/Programs/UnrealVS/UnrealVS.2022/Properties/AssemblyInfo.cs
@@ -31,6 +31,6 @@ using System.Runtime.InteropServices;
 // You can specify all the values or you can default the Revision and Build Numbers 
 // by using the '*' as shown below:
 
-[assembly: AssemblyVersion("1.84.8.1")]
-[assembly: AssemblyFileVersion("1.84.8.1")]
+[assembly: AssemblyVersion("1.84.8.2")]
+[assembly: AssemblyFileVersion("1.84.8.2")]
 
diff --git a/Engine/Source/Programs/UnrealVS/UnrealVS.2022/UnrealVS.2022.csproj b/Engine/Source/Programs/UnrealVS/UnrealVS.2022/UnrealVS.2022.csproj
--- a/Engine/Source/Programs/UnrealVS/UnrealVS.2022/UnrealVS.2022.csproj
+++ b/Engine/Source/Programs/UnrealVS/UnrealVS.2022/UnrealVS.2022.csproj
@@ -171,7 +171,7 @@
       <PrivateAssets>all</PrivateAssets>
     </PackageReference>
     <PackageReference Include="System.Text.Json">
-      <Version>8.0.4</Version>
+      <Version>8.0.5</Version>
     </PackageReference>
   </ItemGroup>
   <ItemGroup>
diff --git a/Engine/Source/Programs/UnrealVS/UnrealVS.Shared/UnrealVSPackage.cs b/Engine/Source/Programs/UnrealVS/UnrealVS.Shared/UnrealVSPackage.cs
--- a/Engine/Source/Programs/UnrealVS/UnrealVS.Shared/UnrealVSPackage.cs
+++ b/Engine/Source/Programs/UnrealVS/UnrealVS.Shared/UnrealVSPackage.cs
@@ -82,7 +82,7 @@ namespace UnrealVS
 	{
 		/** Constants */
 
-		private const string VersionString = "v1.80";
+		private const string VersionString = "v1.84";
 		private const string UnrealSolutionFileNamePrefix = "UE";
 		private const string ExtensionName = "UnrealVS";
 		private const string CommandLineOptionKey = ExtensionName + "CommandLineMRU";
diff --git a/Engine/Source/Programs/UnrealVS/UpdateInstructions.txt b/Engine/Source/Programs/UnrealVS/UpdateInstructions.txt
--- a/Engine/Source/Programs/UnrealVS/UpdateInstructions.txt
+++ b/Engine/Source/Programs/UnrealVS/UpdateInstructions.txt
@@ -4,8 +4,9 @@
 	* Resources are also shared
 	* UnrealVS.vsct is also shared
 	* Manifest file is named source.extension.vsixmanifest and is unique per project
+	* UnrealVSPackage.cs for major/minor version update (this shows up in visual studio version info)
 3. update ReleaseNotes.txt with the changes
-4. Checkout the VSIX in Engine\Extras\UnrealVS\ files and rebuild both UnrealVS.2019 and UnrealVS.2022
+4. Checkout the VSIX in Engine\Extras\UnrealVS\ files and rebuild UnrealVS.2022
 5. Update the atomfeed.xml in Engine\Extras\UnrealVS\<version>\atomfeed.xml
 	* <updated> with new timestamp
 	* <Version> with new version number
