New version of UnrealVS (1.84.9.1) that fixes an issue where running UnrealVS.P4->History would occasionally spawn an empty p4 history window.

- The code path for showing a file's p4 history was still creating the process manually but was not calling ::PullWorkingDirectory meaning that if it was the first command invoked, P4WorkingDirectory would be null.
- Invoking p4vc.bat without setting the working directory to the root of the sln sometimes causes the history window to be empty. I have not tracked the exact cause of this but I'd suspect it is p4config/p4env related but confirming the exact conditions that this occurs in did not seem useful. Providing the correct working directory fixed the problem for the one 100% reproducible case I had access too.
- We have a lot of helper methods for launching different p4 processes all of which handle things like this so it was easiest to remove the manual process creation and just call ::TryP4VCCommand instead.

diff --git a/Engine/Extras/UnrealVS/VS2022/AtomFeed.xml b/Engine/Extras/UnrealVS/VS2022/AtomFeed.xml
--- a/Engine/Extras/UnrealVS/VS2022/AtomFeed.xml
+++ b/Engine/Extras/UnrealVS/VS2022/AtomFeed.xml
@@ -2,13 +2,13 @@
 <feed xmlns="http://www.w3.org/2005/Atom">
   <title type="text">UnrealVS Gallery</title>
   <id>uuid:45b4571a-2981-49f1-9f70-04cfa96db5c8;id=1</id>
-  <updated>2025-04-28T23:33:26Z</updated>
+  <updated>2025-05-02T11:21:31Z</updated>
   <entry>
     <id>ddbf523f-7eb6-4887-bd51-85a714ff87eb</id>
     <title type="text">UnrealVS</title>
     <summary type="text">Unreal Engine extension for Visual Studio</summary>
     <published>2022-10-21T10:01:57-07:00</published>
-    <updated>2025-04-28T23:33:26-00:00</updated>
+    <updated>2025-05-02T11:21:31+00:00</updated>
     <author>
       <name>Epic Games, Inc.</name>
     </author>
@@ -19,7 +19,7 @@
     <content type="application/octet-stream" src="UnrealVS.vsix" />
     <Vsix xmlns:xsd="http://www.w3.org/2001/XMLSchema" xmlns:xsi="http://www.w3.org/2001/XMLSchema-instance" xmlns="http://schemas.microsoft.com/developer/vsx-syndication-schema/2010">
       <Id>ddbf523f-7eb6-4887-bd51-85a714ff87eb</Id>
-      <Version>1.84.9.0</Version>
+      <Version>1.84.9.1</Version>
       <Rating xsi:nil="true" />
       <RatingCount xsi:nil="true" />
       <DownloadCount xsi:nil="true" />
diff --git a/Engine/Source/Programs/UnrealVS/ReleaseNotes.txt b/Engine/Source/Programs/UnrealVS/ReleaseNotes.txt
--- a/Engine/Source/Programs/UnrealVS/ReleaseNotes.txt
+++ b/Engine/Source/Programs/UnrealVS/ReleaseNotes.txt
@@ -1,3 +1,6 @@
+1.84.9.1
+* Fix issue where running UnrealVS.P4->History would occasionally spawn an empty p4 history window.
+
 1.84.9.0
 * Display errors during test discovery in the Tests output log.
 * Don't run AutoSDK when executing tests from the Test Explorer.
diff --git a/Engine/Source/Programs/UnrealVS/UnrealVS.2022/Properties/AssemblyInfo.cs b/Engine/Source/Programs/UnrealVS/UnrealVS.2022/Properties/AssemblyInfo.cs
--- a/Engine/Source/Programs/UnrealVS/UnrealVS.2022/Properties/AssemblyInfo.cs
+++ b/Engine/Source/Programs/UnrealVS/UnrealVS.2022/Properties/AssemblyInfo.cs
@@ -31,6 +31,6 @@ using System.Runtime.InteropServices;
 // You can specify all the values or you can default the Revision and Build Numbers 
 // by using the '*' as shown below:
 
-[assembly: AssemblyVersion("1.84.9.0")]
-[assembly: AssemblyFileVersion("1.84.9.0")]
+[assembly: AssemblyVersion("1.84.9.1")]
+[assembly: AssemblyFileVersion("1.84.9.1")]
 
diff --git a/Engine/Source/Programs/UnrealVS/UnrealVS.Shared/P4Commands.cs b/Engine/Source/Programs/UnrealVS/UnrealVS.Shared/P4Commands.cs
--- a/Engine/Source/Programs/UnrealVS/UnrealVS.Shared/P4Commands.cs
+++ b/Engine/Source/Programs/UnrealVS/UnrealVS.Shared/P4Commands.cs
@@ -511,20 +511,7 @@ namespace UnrealVS
 			}
 
 			string Command = $"history \"{DTE.ActiveDocument.FullName}\"";
-			P4OutputPane.OutputStringThreadSafe($"1>------ Running P4FileHistory with args: {Command}{Environment.NewLine}");
-
-			Process ChildProcess = new Process()
-			{
-				StartInfo = new ProcessStartInfo()
-				{
-					FileName = P4VCCmd,
-					Arguments = Command,
-					WorkingDirectory = P4WorkingDirectory,
-					UseShellExecute = false,
-					CreateNoWindow = true
-				}
-			};
-			ChildProcess.Start();
+			TryP4VCCommand(Command);
 		}
 
 		private void P4IntegrationAwareTimeLapseHandler(object Sender, EventArgs Args)
